Remove CodeWorks for Android and add setup batch file

#ue4
#android

diff --git a/Engine/Extras/Android/SetupAndroid.bat b/Engine/Extras/Android/SetupAndroid.bat
new file mode 100644
--- /dev/null
+++ b/Engine/Extras/Android/SetupAndroid.bat
@@ -0,0 +1,74 @@
+@echo off
+setlocal ENABLEEXTENSIONS
+set KEY_NAME=HKLM\SOFTWARE\Android Studio
+set VALUE_NAME=Path
+set STUDIO_PATH=""
+
+FOR /F "tokens=2*" %%A IN ('REG.exe query "%KEY_NAME%" /v "%VALUE_NAME%"') DO (set STUDIO_PATH="%%B")
+
+if %STUDIO_PATH% == "" (
+	echo Android Studio not installed, please download Android Studio 3.5.3 from https://developer.android.com/studio
+	pause
+	exit /b 1
+)
+echo Android Studio Path: %STUDIO_PATH%
+
+set VALUE_NAME=SdkPath
+set STUDIO_SDK_PATH=""
+FOR /F "tokens=2*" %%A IN ('REG.exe query "%KEY_NAME%" /v "%VALUE_NAME%"') DO (set STUDIO_SDK_PATH="%%B")
+
+set ANDROID_LOCAL=%LOCALAPPDATA%\Android\Sdk
+
+if %STUDIO_SDK_PATH% == "" (
+	IF EXIST "%ANDROID_LOCAL%" (
+		set STUDIO_SDK_PATH=%ANDROID_LOCAL%
+	) ELSE (
+		IF EXIST "%ANDROID_HOME%" (
+			set STUDIO_SDK_PATH=%ANDROID_HOME%
+		) ELSE (
+			echo Unable to locate local Android SDK location. Did you run Android Studio after installing?
+			pause
+			exit /b 1
+		)
+	)
+)
+echo Android Studio SDK Path: %STUDIO_SDK_PATH%
+
+if DEFINED ANDROID_HOME (set a=1) ELSE (
+	set ANDROID_HOME=%STUDIO_SDK_PATH%
+	setx ANDROID_HOME %STUDIO_SDK_PATH%
+)
+if DEFINED JAVA_HOME (set a=1) ELSE (
+	set JAVA_HOME=%STUDIO_PATH%\jre
+	setx JAVA_HOME %STUDIO_PATH%\jre
+)
+set NDKINSTALLPATH=%STUDIO_SDK_PATH%\ndk\21.0.6113669
+set PLATFORMTOOLS=%STUDIO_SDK_PATH%\platform-tools;%STUDIO_SDK_PATH%\tools
+
+where.exe /Q adb.exe
+IF /I "%ERRORLEVEL%" NEQ "0" (
+	FOR /F "tokens=2*" %%A IN ('REG.exe query "HKCU\Environment" /v "Path"') DO (set USERPATH="%%B")
+	echo %USERPATH%
+	setx PATH "%USERPATH%;%PLATFORMTOOLS%"
+	echo Added %PLATFORMTOOLS% to path
+)
+
+call %STUDIO_SDK_PATH%\tools\bin\sdkmanager.bat "platform-tools" "platforms;android-28" "build-tools;28.0.3" "lldb;3.1" "cmake;3.10.2.4988404" "ndk;21.0.6113669"
+
+IF /I "%ERRORLEVEL%" NEQ "0" (
+	echo Update failed. Please check the Android Studio install.
+	pause
+	exit /b 1
+)
+
+if EXIST %NDKINSTALLPATH% (
+	echo Success!
+	setx NDKROOT %NDKINSTALLPATH%
+) ELSE (
+	echo Update failed. Did you accept the license agreement?
+	pause
+	exit /b 1
+)
+
+pause
+exit /b 0
