Fix batch file variable declaration

#ue4
#android

diff --git a/Engine/Extras/Android/SetupAndroid.bat b/Engine/Extras/Android/SetupAndroid.bat
--- a/Engine/Extras/Android/SetupAndroid.bat
+++ b/Engine/Extras/Android/SetupAndroid.bat
@@ -45,10 +45,15 @@ if DEFINED JAVA_HOME (set a=1) ELSE (
 set NDKINSTALLPATH=%STUDIO_SDK_PATH%\ndk\21.0.6113669
 set PLATFORMTOOLS=%STUDIO_SDK_PATH%\platform-tools;%STUDIO_SDK_PATH%\tools
 
+set KEY_NAME=HKCU\Environment
+set VALUE_NAME=Path
+set USERPATH=""
+
+FOR /F "tokens=2*" %%A IN ('REG.exe query "%KEY_NAME%" /v "%VALUE_NAME%"') DO (set USERPATH="%%B")
+
 where.exe /Q adb.exe
 IF /I "%ERRORLEVEL%" NEQ "0" (
-	FOR /F "tokens=2*" %%A IN ('REG.exe query "HKCU\Environment" /v "Path"') DO (set USERPATH="%%B")
-	echo %USERPATH%
+	echo Current user path: %USERPATH%
 	setx PATH "%USERPATH%;%PLATFORMTOOLS%"
 	echo Added %PLATFORMTOOLS% to path
 )
