Add Mac version of SetupAndroid script

#ue4
#android

diff --git a/Engine/Extras/Android/SetupAndroid.command b/Engine/Extras/Android/SetupAndroid.command
new file mode 100755
--- /dev/null
+++ b/Engine/Extras/Android/SetupAndroid.command
@@ -0,0 +1,86 @@
+#!/bin/sh
+
+STUDIO_PATH="~/Applications/Android Studio.app"
+if [ ! -x "$STUDIO_PATH" ]; then
+	STUDIO_PATH="/Applications/Android Studio.app"
+	if [ ! -x "$STUDIO_PATH" ]; then
+		echo Android Studio not installed, please download Android Studio 3.5.3 from https://developer.android.com/studio
+		read -rsp $'Press any key to continue...\n' -n1 key
+		exit 1
+	fi
+fi
+echo Android Studio Path: $STUDIO_PATH
+
+if [ "$STUDIO_SDK_PATH" == "" ]; then
+	STUDIO_SDK_PATH=$ANDROID_HOME
+	if [ "$STUDIO_SDK_PATH" == "" ]; then
+		STUDIO_SDK_PATH=~/Library/Android/sdk
+	fi
+fi
+if [ "$1" != "" ]; then
+	STUDIO_SDK_PATH=$1
+fi
+if [ ! -d "$STUDIO_SDK_PATH" ]; then
+	echo Android SDK not found at: $STUDIO_SDK_PATH
+	echo Unable to locate local Android SDK location. Did you run Android Studio after installing?
+	echo If Android Studio is installed, please run again with SDK path as parameter, otherwise download Android Studio 3.5.3 from https://developer.android.com/studio
+	read -rsp $'Press any key to continue...\n' -n1 key
+	exit 1
+fi
+echo Android Studio SDK Path: $STUDIO_SDK_PATH
+
+if ! grep -q "export ANDROID_HOME=\"$STUDIO_SDK_PATH\"" ~/.bash_profile 
+then
+	echo "export ANDROID_HOME=\"$STUDIO_SDK_PATH\"" >>~/.bash_profile
+fi
+
+export JAVA_HOME="$STUDIO_PATH/Contents/jre/jdk/Contents/Home"
+if ! grep -q "export JAVA_HOME=\"$JAVA_HOME\"" ~/.bash_profile
+then
+	echo "export JAVA_HOME=\"$JAVA_HOME\"" >>~/.bash_profile
+fi
+NDKINSTALLPATH="$STUDIO_SDK_PATH/ndk/21.0.6113669"
+PLATFORMTOOLS="$STUDIO_SDK_PATH/platform-tools:$STUDIO_SDK_PATH/build-tools/28.0.3:$%STUDIO_SDK_PATH/tools/bin"
+
+retVal=$(type -P "adb")
+if [ $retVal == "" ]; then
+	echo 'PATH="$PATH:$PLATFORMTOOLS"' >>~/.bash_profile
+	echo Added $PLATFORMTOOLS to path
+fi
+
+"$STUDIO_SDK_PATH/tools/bin/sdkmanager" "platform-tools" "platforms;android-28" "build-tools;28.0.3" "lldb;3.1" "cmake;3.10.2.4988404" "ndk;21.0.6113669"
+
+retVal=$?
+if [ $retVal -ne 0 ]; then
+	echo Update failed. Please check the Android Studio install.
+	read -rsp $'Press any key to continue...\n' -n1 key
+	exit $retVal
+fi
+
+if [ ! -d "$STUDIO_SDK_PATH/platform-tools" ]; then
+	retVal=1
+fi
+if [ ! -d "$STUDIO_SDK_PATH/platforms/android-28" ]; then
+	retVal=1
+fi
+if [ ! -d "$STUDIO_SDK_PATH/build-tools/28.0.3" ]; then
+	retVal=1
+fi
+if [ ! -d "$NDKINSTALLPATH" ]; then
+	retVal=1
+fi
+
+if [ $retVal -ne 0 ]; then
+	echo Update failed. Did you accept the license agreement?
+	read -rsp $'Press any key to continue...\n' -n1 key
+	exit $retVal
+fi
+
+echo Success!
+
+if ! grep -q "export NDKROOT=\"$NDKINSTALLPATH\"" ~/.bash_profile
+then
+	echo "export NDKROOT=\"$NDKINSTALLPATH\"" >>~/.bash_profile
+fi
+
+exit 0
