export PATH for Linux/Mac SetupAndroid scripts

#android

diff --git a/Engine/Extras/Android/SetupAndroid.command b/Engine/Extras/Android/SetupAndroid.command
--- a/Engine/Extras/Android/SetupAndroid.command
+++ b/Engine/Extras/Android/SetupAndroid.command
@@ -44,7 +44,7 @@ PLATFORMTOOLS="$STUDIO_SDK_PATH/platform-tools:$STUDIO_SDK_PATH/build-tools/28.0
 
 retVal=$(type -P "adb")
 if [ $retVal == "" ]; then
-	echo PATH="\"\$PATH:$PLATFORMTOOLS\"" >>~/.bash_profile
+	echo export PATH="\"\$PATH:$PLATFORMTOOLS\"" >>~/.bash_profile
 	echo Added $PLATFORMTOOLS to path
 fi
 
diff --git a/Engine/Extras/Android/SetupAndroid.sh b/Engine/Extras/Android/SetupAndroid.sh
--- a/Engine/Extras/Android/SetupAndroid.sh
+++ b/Engine/Extras/Android/SetupAndroid.sh
@@ -42,7 +42,7 @@ PLATFORMTOOLS="$STUDIO_SDK_PATH/platform-tools:$STUDIO_SDK_PATH/build-tools/28.0
 
 retVal=$(type -P "adb")
 if [ "$retVal" == "" ]; then
-	echo PATH="\"\$PATH:$PLATFORMTOOLS\"" >>$HOME/.bashrc
+	echo export PATH="\"\$PATH:$PLATFORMTOOLS\"" >>$HOME/.bashrc
 	echo Added $PLATFORMTOOLS to path
 fi
 
