Add alternative location for sdkmanager (moved with newer Android Studio isntalls)

#ue4
#android

diff --git a/Engine/Extras/Android/SetupAndroid.bat b/Engine/Extras/Android/SetupAndroid.bat
--- a/Engine/Extras/Android/SetupAndroid.bat
+++ b/Engine/Extras/Android/SetupAndroid.bat
@@ -58,7 +58,21 @@ IF /I "%ERRORLEVEL%" NEQ "0" (
 	echo Added %PLATFORMTOOLS% to path
 )
 
-call %STUDIO_SDK_PATH%\tools\bin\sdkmanager.bat "platform-tools" "platforms;android-28" "build-tools;28.0.3" "lldb;3.1" "cmake;3.10.2.4988404" "ndk;21.0.6113669"
+set SDKMANAGER=%STUDIO_SDK_PATH%\tools\bin\sdkmanager.bat
+IF EXIST "%SDKMANAGER%" (
+	echo Using sdkmanager: %SDKMANAGER%
+) ELSE (
+	set SDKMANAGER=%STUDIO_SDK_PATH%\cmdline-tools\latest\bin\sdkmanager.bat
+	IF EXIST "%SDKMANAGER%" (
+		echo Using sdkmanager: %SDKMANAGER%
+	) ELSE (
+		echo Unable to locate sdkmanager.bat. Did you run Android Studio and install cmdline-tools after installing?
+		pause
+		exit /b 1
+	)
+)
+
+call "%SDKMANAGER%" "platform-tools" "platforms;android-28" "build-tools;28.0.3" "lldb;3.1" "cmake;3.10.2.4988404" "ndk;21.0.6113669"
 
 IF /I "%ERRORLEVEL%" NEQ "0" (
 	echo Update failed. Please check the Android Studio install.
diff --git a/Engine/Extras/Android/SetupAndroid.command b/Engine/Extras/Android/SetupAndroid.command
--- a/Engine/Extras/Android/SetupAndroid.command
+++ b/Engine/Extras/Android/SetupAndroid.command
@@ -48,7 +48,18 @@ if [ $retVal == "" ]; then
 	echo Added $PLATFORMTOOLS to path
 fi
 
-"$STUDIO_SDK_PATH/tools/bin/sdkmanager" "platform-tools" "platforms;android-28" "build-tools;28.0.3" "lldb;3.1" "cmake;3.10.2.4988404" "ndk;21.0.6113669"
+SDKMANAGERPATH="$STUDIO_SDK_PATH/tools/bin"
+if [ ! -d "$SDKMANAGERPATH" ]; then
+	SDKMANAGERPATH="$STUDIO_SDK_PATH/cmdline-tools/latest/bintools/bin"
+	if [ ! -d "$SDKMANAGERPATH" ]; then
+		echo Unable to locate sdkmanager.bat. Did you run Android Studio and install cmdline-tools after installing?
+		SDKMANAGERPATH="$STUDIO_SDK_PATH/cmdline-tools/latest/bintools/bin"
+		read -rsp $'Press any key to continue...\n' -n1 key
+		exit 1
+	fi
+fi
+
+"$SDKMANAGERPATH/sdkmanager" "platform-tools" "platforms;android-28" "build-tools;28.0.3" "lldb;3.1" "cmake;3.10.2.4988404" "ndk;21.0.6113669"
 
 retVal=$?
 if [ $retVal -ne 0 ]; then
diff --git a/Engine/Extras/Android/SetupAndroid.sh b/Engine/Extras/Android/SetupAndroid.sh
--- a/Engine/Extras/Android/SetupAndroid.sh
+++ b/Engine/Extras/Android/SetupAndroid.sh
@@ -46,7 +46,18 @@ if [ "$retVal" == "" ]; then
 	echo Added $PLATFORMTOOLS to path
 fi
 
-"$STUDIO_SDK_PATH/tools/bin/sdkmanager" "platform-tools" "platforms;android-28" "build-tools;28.0.3" "lldb;3.1" "cmake;3.10.2.4988404" "ndk;21.0.6113669"
+SDKMANAGERPATH="$STUDIO_SDK_PATH/tools/bin"
+if [ ! -d "$SDKMANAGERPATH" ]; then
+	SDKMANAGERPATH="$STUDIO_SDK_PATH/cmdline-tools/latest/bintools/bin"
+	if [ ! -d "$SDKMANAGERPATH" ]; then
+		echo Unable to locate sdkmanager.bat. Did you run Android Studio and install cmdline-tools after installing?
+		SDKMANAGERPATH="$STUDIO_SDK_PATH/cmdline-tools/latest/bintools/bin"
+		read -rsp $'Press any key to continue...\n' -n1 key
+		exit 1
+	fi
+fi
+
+"$SDKMANAGERPATH/sdkmanager" "platform-tools" "platforms;android-28" "build-tools;28.0.3" "lldb;3.1" "cmake;3.10.2.4988404" "ndk;21.0.6113669"
 
 retVal=$?
 if [ $retVal -ne 0 ]; then
