Fix android setup script for Linux/Mac

diff --git a/Engine/Extras/Android/SetupAndroid.command b/Engine/Extras/Android/SetupAndroid.command
--- a/Engine/Extras/Android/SetupAndroid.command
+++ b/Engine/Extras/Android/SetupAndroid.command
@@ -50,10 +50,9 @@ fi
 
 SDKMANAGERPATH="$STUDIO_SDK_PATH/tools/bin"
 if [ ! -d "$SDKMANAGERPATH" ]; then
-	SDKMANAGERPATH="$STUDIO_SDK_PATH/cmdline-tools/latest/bintools/bin"
+	SDKMANAGERPATH="$STUDIO_SDK_PATH/cmdline-tools/latest/bin"
 	if [ ! -d "$SDKMANAGERPATH" ]; then
-		echo Unable to locate sdkmanager.bat. Did you run Android Studio and install cmdline-tools after installing?
-		SDKMANAGERPATH="$STUDIO_SDK_PATH/cmdline-tools/latest/bintools/bin"
+		echo Unable to locate sdkmanager. Did you run Android Studio and install cmdline-tools after installing?
 		read -rsp $'Press any key to continue...\n' -n1 key
 		exit 1
 	fi
diff --git a/Engine/Extras/Android/SetupAndroid.sh b/Engine/Extras/Android/SetupAndroid.sh
--- a/Engine/Extras/Android/SetupAndroid.sh
+++ b/Engine/Extras/Android/SetupAndroid.sh
@@ -48,10 +48,9 @@ fi
 
 SDKMANAGERPATH="$STUDIO_SDK_PATH/tools/bin"
 if [ ! -d "$SDKMANAGERPATH" ]; then
-	SDKMANAGERPATH="$STUDIO_SDK_PATH/cmdline-tools/latest/bintools/bin"
+	SDKMANAGERPATH="$STUDIO_SDK_PATH/cmdline-tools/latest/bin"
 	if [ ! -d "$SDKMANAGERPATH" ]; then
-		echo Unable to locate sdkmanager.bat. Did you run Android Studio and install cmdline-tools after installing?
-		SDKMANAGERPATH="$STUDIO_SDK_PATH/cmdline-tools/latest/bintools/bin"
+		echo Unable to locate sdkmanager. Did you run Android Studio and install cmdline-tools after installing?
 		read -rsp $'Press any key to continue...\n' -n1 key
 		exit 1
 	fi
