Remove check for ANDROID_HOME from Mac and Linux
Run SetupAndroid.command if SetupAndroid.sh is run on Mac instead

#ue4
#android

diff --git a/Engine/Extras/Android/SetupAndroid.command b/Engine/Extras/Android/SetupAndroid.command
--- a/Engine/Extras/Android/SetupAndroid.command
+++ b/Engine/Extras/Android/SetupAndroid.command
@@ -12,10 +12,7 @@ fi
 echo Android Studio Path: $STUDIO_PATH
 
 if [ "$STUDIO_SDK_PATH" == "" ]; then
-	STUDIO_SDK_PATH=$ANDROID_HOME
-	if [ "$STUDIO_SDK_PATH" == "" ]; then
-		STUDIO_SDK_PATH=~/Library/Android/sdk
-	fi
+	STUDIO_SDK_PATH=~/Library/Android/sdk
 fi
 if [ "$1" != "" ]; then
 	STUDIO_SDK_PATH=$1
diff --git a/Engine/Extras/Android/SetupAndroid.sh b/Engine/Extras/Android/SetupAndroid.sh
--- a/Engine/Extras/Android/SetupAndroid.sh
+++ b/Engine/Extras/Android/SetupAndroid.sh
@@ -1,5 +1,12 @@
 #!/bin/bash
 
+if [[ "$OSTYPE" == "darwin"* ]]; then
+	echo "Please run SetupAndroid.command on MacOSX; attempting to run it for you."
+	DIR="$( cd "$( dirname "${BASH_SOURCE[0]}" )" >/dev/null 2>&1 && pwd )"
+	exec "$DIR"/SetupAndroid.command
+	exit 1
+fi
+
 STUDIO_PATH="$HOME/android-studio"
 if [ ! -d "$STUDIO_PATH" ]; then
 	echo "Android Studio not installed, please download Android Studio 3.5.3 from https://developer.android.com/studio"
@@ -10,10 +17,7 @@ fi
 echo Android Studio Path: $STUDIO_PATH
 
 if [ "$STUDIO_SDK_PATH" == "" ]; then
-	STUDIO_SDK_PATH=$ANDROID_HOME
-	if [ "$STUDIO_SDK_PATH" == "" ]; then
-		STUDIO_SDK_PATH=$HOME/Android/Sdk
-	fi
+	STUDIO_SDK_PATH=$HOME/Android/Sdk
 fi
 if [ "$1" != "" ]; then
 	STUDIO_SDK_PATH=$1
