Add new line before any echo in SetupAndroid

diff --git a/Engine/Extras/Android/SetupAndroid.command b/Engine/Extras/Android/SetupAndroid.command
--- a/Engine/Extras/Android/SetupAndroid.command
+++ b/Engine/Extras/Android/SetupAndroid.command
@@ -28,12 +28,14 @@ echo Android Studio SDK Path: $STUDIO_SDK_PATH
 
 if ! grep -q "export ANDROID_HOME=\"$STUDIO_SDK_PATH\"" ~/.bash_profile 
 then
+	echo >> ~/.bash_profile
 	echo "export ANDROID_HOME=\"$STUDIO_SDK_PATH\"" >>~/.bash_profile
 fi
 
 export JAVA_HOME="$STUDIO_PATH/Contents/jre/jdk/Contents/Home"
 if ! grep -q "export JAVA_HOME=\"$JAVA_HOME\"" ~/.bash_profile
 then
+	echo >> ~/.bash_profile
 	echo "export JAVA_HOME=\"$JAVA_HOME\"" >>~/.bash_profile
 fi
 NDKINSTALLPATH="$STUDIO_SDK_PATH/ndk/21.1.6352462"
@@ -41,6 +43,7 @@ PLATFORMTOOLS="$STUDIO_SDK_PATH/platform-tools:$STUDIO_SDK_PATH/build-tools/28.0
 
 retVal=$(type -P "adb")
 if [ $retVal == "" ]; then
+	echo >> ~/.bash_profile
 	echo export PATH="\"\$PATH:$PLATFORMTOOLS\"" >>~/.bash_profile
 	echo Added $PLATFORMTOOLS to path
 fi
@@ -87,6 +90,7 @@ echo Success!
 
 if ! grep -q "export NDKROOT=\"$NDKINSTALLPATH\"" ~/.bash_profile
 then
+	echo >> ~/.bash_profile
 	echo "export NDKROOT=\"$NDKINSTALLPATH\"" >>~/.bash_profile
 	echo "export NDK_ROOT=\"$NDKINSTALLPATH\"" >>~/.bash_profile
 fi
diff --git a/Engine/Extras/Android/SetupAndroid.sh b/Engine/Extras/Android/SetupAndroid.sh
--- a/Engine/Extras/Android/SetupAndroid.sh
+++ b/Engine/Extras/Android/SetupAndroid.sh
@@ -33,12 +33,14 @@ echo Android Studio SDK Path: $STUDIO_SDK_PATH
 
 if ! grep -q "export ANDROID_HOME=\"$STUDIO_SDK_PATH\"" $HOME/.bashrc
 then
+	echo >>$HOME/.bashrc
 	echo "export ANDROID_HOME=\"$STUDIO_SDK_PATH\"" >>$HOME/.bashrc
 fi
 
 export JAVA_HOME="$STUDIO_PATH/jre"
 if ! grep -q "export JAVA_HOME=\"$JAVA_HOME\"" $HOME/.bashrc
 then
+	echo >>$HOME/.bashrc
 	echo "export JAVA_HOME=\"$JAVA_HOME\"" >>$HOME/.bashrc
 fi
 NDKINSTALLPATH="$STUDIO_SDK_PATH/ndk/21.1.6352462"
@@ -46,6 +48,7 @@ PLATFORMTOOLS="$STUDIO_SDK_PATH/platform-tools:$STUDIO_SDK_PATH/build-tools/28.0
 
 retVal=$(type -P "adb")
 if [ "$retVal" == "" ]; then
+	echo >>$HOME/.bashrc
 	echo export PATH="\"\$PATH:$PLATFORMTOOLS\"" >>$HOME/.bashrc
 	echo Added $PLATFORMTOOLS to path
 fi
@@ -92,6 +95,7 @@ echo Success!
 
 if ! grep -q "export NDKROOT=\"$NDKINSTALLPATH\"" $HOME/.bashrc
 then
+	echo >>$HOME/.bashrc
 	echo "export NDKROOT=\"$NDKINSTALLPATH\"" >>$HOME/.bashrc
 	echo "export NDK_ROOT=\"$NDKINSTALLPATH\"" >>$HOME/.bashrc
 fi
