Fix extra % added to ENV path

diff --git a/Engine/Extras/Android/SetupAndroid.command b/Engine/Extras/Android/SetupAndroid.command
--- a/Engine/Extras/Android/SetupAndroid.command
+++ b/Engine/Extras/Android/SetupAndroid.command
@@ -71,7 +71,7 @@ then
 	echo "export JAVA_HOME=\"$JAVA_HOME\"" >>~/.bash_profile
 fi
 NDKINSTALLPATH="$STUDIO_SDK_PATH/ndk/${NDK_VERSION}"
-PLATFORMTOOLS="$STUDIO_SDK_PATH/platform-tools:$STUDIO_SDK_PATH/build-tools/${BUILDTOOLS_VERSION}:$%STUDIO_SDK_PATH/tools/bin"
+PLATFORMTOOLS="$STUDIO_SDK_PATH/platform-tools:$STUDIO_SDK_PATH/build-tools/${BUILDTOOLS_VERSION}:$STUDIO_SDK_PATH/tools/bin"
 
 retVal=$(type -P "adb")
 if [ "$retVal" == "" ]; then
diff --git a/Engine/Extras/Android/SetupAndroid.sh b/Engine/Extras/Android/SetupAndroid.sh
--- a/Engine/Extras/Android/SetupAndroid.sh
+++ b/Engine/Extras/Android/SetupAndroid.sh
@@ -77,7 +77,7 @@ then
 	echo "export JAVA_HOME=\"$JAVA_HOME\"" >>$HOME/.bashrc
 fi
 NDKINSTALLPATH="$STUDIO_SDK_PATH/ndk/${NDK_VERSION}"
-PLATFORMTOOLS="$STUDIO_SDK_PATH/platform-tools:$STUDIO_SDK_PATH/build-tools/${BUILDTOOLS_VERSION}:$%STUDIO_SDK_PATH/tools/bin"
+PLATFORMTOOLS="$STUDIO_SDK_PATH/platform-tools:$STUDIO_SDK_PATH/build-tools/${BUILDTOOLS_VERSION}:$STUDIO_SDK_PATH/tools/bin"
 
 retVal=$(type -P "adb")
 if [ "$retVal" == "" ]; then
