Fix SetupAndroid.bat to handle unexpected characters in user path

#ue4
#android

diff --git a/Engine/Extras/Android/SetupAndroid.bat b/Engine/Extras/Android/SetupAndroid.bat
--- a/Engine/Extras/Android/SetupAndroid.bat
+++ b/Engine/Extras/Android/SetupAndroid.bat
@@ -35,7 +35,7 @@ echo Android Studio Path: %STUDIO_PATH%
 
 set VALUE_NAME=SdkPath
 set STUDIO_SDK_PATH=
-FOR /F "tokens=2*" %%A IN ('REG.exe query "%KEY_NAME%" /v "%VALUE_NAME%"') DO (set STUDIO_SDK_PATH=%%B)
+FOR /F "tokens=2* usebackq" %%A IN ('REG.exe query "%KEY_NAME%" /v "%VALUE_NAME%"') DO (set STUDIO_SDK_PATH=%%B)
 
 set ANDROID_LOCAL=%LOCALAPPDATA%\Android\Sdk
 
