PR #8232: fix fallback location check in SetupAndroid.bat (Contributed by [REDACTED_PERSON])

#ue5
#android

diff --git a/Engine/Extras/Android/SetupAndroid.bat b/Engine/Extras/Android/SetupAndroid.bat
--- a/Engine/Extras/Android/SetupAndroid.bat
+++ b/Engine/Extras/Android/SetupAndroid.bat
@@ -1,5 +1,6 @@
 @echo off
 setlocal ENABLEEXTENSIONS
+setlocal ENABLEDELAYEDEXPANSION
 set KEY_NAME=HKLM\SOFTWARE\Android Studio
 set VALUE_NAME=Path
 set STUDIO_PATH=
@@ -83,8 +84,8 @@ IF EXIST "%SDKMANAGER%" (
 	echo Using sdkmanager: %SDKMANAGER%
 ) ELSE (
 	set SDKMANAGER=%STUDIO_SDK_PATH%\cmdline-tools\latest\bin\sdkmanager.bat
-	IF EXIST "%SDKMANAGER%" (
-		echo Using sdkmanager: %SDKMANAGER%
+	IF EXIST "!SDKMANAGER!" (
+		echo Using sdkmanager: !SDKMANAGER!
 	) ELSE (
 		echo Unable to locate sdkmanager.bat. Did you run Android Studio and install cmdline-tools after installing?
 		%PAUSE%
