Mac: fix installation of Android NDK throught the script SetupAndroid.command
* if default shell is zsh,  use .zprofile as well as .bash_profile (as UE is reading .bash_profile)
* make sure env variables are set up properly (prevent duplication if script was already launch)
* change rem to '#' as this was causing a warning

diff --git a/Engine/Extras/Android/SetupAndroid.command b/Engine/Extras/Android/SetupAndroid.command
--- a/Engine/Extras/Android/SetupAndroid.command
+++ b/Engine/Extras/Android/SetupAndroid.command
@@ -11,13 +11,45 @@ pausefunc()
     read -rsp $'Press any key to continue...\n' -n1 key
 }
 
+# $1 parameter should be the profile file: .bash_profile as it is read by UE and .zprofile if default shell is zsh
+updateProfileFileFunc()
+{
+	if ! grep -q "export ANDROID_HOME=\"$STUDIO_SDK_PATH\"" "$1" 
+	then
+		echo >> "$1"
+		echo "export ANDROID_HOME=\"$STUDIO_SDK_PATH\"" >> "$1"
+	fi
+ 
+	if ! grep -q "export JAVA_HOME=\"$JAVA_HOME\"" "$1"
+	then
+		echo >> $1
+		echo "export JAVA_HOME=\"$JAVA_HOME\"" >> "$1"
+	fi
+ 
+	if [ "$adbPath" == "" ]; then
+		if ! grep -q "export PATH=\"\$PATH:$PLATFORMTOOLS\"" "$1"
+		then
+			echo >> "$1"
+			echo export PATH="\"\$PATH:$PLATFORMTOOLS\"" >> "$1"
+			echo Added $PLATFORMTOOLS to PATH in $1 
+		fi
+	fi
+ 
+	if ! grep -q "export NDKROOT=\"$NDKINSTALLPATH\"" "$1"
+	then
+		echo >> "$1"
+		echo "export NDKROOT=\"$NDKINSTALLPATH\"" >> "$1"
+		echo "export NDK_ROOT=\"$NDKINSTALLPATH\"" >> "$1"
+	fi
+}
+
 if [[ ${ARG5} == "-noninteractive" ]]; then
     PAUSE=
 else
     PAUSE="pausefunc"
 fi
 
-rem hardcoded versions for compatibility with non-Turnkey manual running
+# hardcoded versions for compatibility with non-Turnkey manual running
 if [[ -z "${PLATFORMS_VERSION}" ]]; then
     PLATFORMS_VERSION="android-30"
 fi
@@ -58,28 +90,10 @@ if [ ! -d "$STUDIO_SDK_PATH" ]; then
 fi
 echo Android Studio SDK Path: $STUDIO_SDK_PATH
 
-if ! grep -q "export ANDROID_HOME=\"$STUDIO_SDK_PATH\"" ~/.bash_profile 
-then
-	echo >> ~/.bash_profile
-	echo "export ANDROID_HOME=\"$STUDIO_SDK_PATH\"" >>~/.bash_profile
-fi
-
 export JAVA_HOME="$STUDIO_PATH/Contents/jre/jdk/Contents/Home"
-if ! grep -q "export JAVA_HOME=\"$JAVA_HOME\"" ~/.bash_profile
-then
-	echo >> ~/.bash_profile
-	echo "export JAVA_HOME=\"$JAVA_HOME\"" >>~/.bash_profile
-fi
 NDKINSTALLPATH="$STUDIO_SDK_PATH/ndk/${NDK_VERSION}"
 PLATFORMTOOLS="$STUDIO_SDK_PATH/platform-tools:$STUDIO_SDK_PATH/build-tools/${BUILDTOOLS_VERSION}:$STUDIO_SDK_PATH/tools/bin"
-
-retVal=$(type -P "adb")
-if [ "$retVal" == "" ]; then
-	echo >> ~/.bash_profile
-	echo export PATH="\"\$PATH:$PLATFORMTOOLS\"" >>~/.bash_profile
-	echo Added $PLATFORMTOOLS to path
-fi
-
+adbPath=$(type -P "adb")
 SDKMANAGERPATH="$STUDIO_SDK_PATH/tools/bin"
 if [ ! -d "$SDKMANAGERPATH" ]; then
 	SDKMANAGERPATH="$STUDIO_SDK_PATH/cmdline-tools/latest/bin"
@@ -118,13 +132,18 @@ if [ $retVal -ne 0 ]; then
 	exit $retVal
 fi
 
-echo Success!
-
-if ! grep -q "export NDKROOT=\"$NDKINSTALLPATH\"" ~/.bash_profile
-then
-	echo >> ~/.bash_profile
-	echo "export NDKROOT=\"$NDKINSTALLPATH\"" >>~/.bash_profile
-	echo "export NDK_ROOT=\"$NDKINSTALLPATH\"" >>~/.bash_profile
+# on zsh, ~/.bash_profile may not exist
+if [ ! -f ~/.bash_profile ]; then
+	echo > ~/.bash_profile
+fi
+ 
+updateProfileFileFunc ~/.bash_profile
+# on zsh, update .zprofile and create .bash_profile (as UE reads it)
+if [ $SHELL == "/bin/zsh" ]; then
+	updateProfileFileFunc ~/.zprofile
+	echo Adding UE Android dev env variables to .zprofile
 fi
 
-exit 0
+echo Success!
+
+exit 0
\ No newline at end of file
