PR #9665: Add ANDROID_SDK_HOME for AGDE (Contributed by [REDACTED_PERSON])

#android

diff --git a/Engine/Extras/Android/SetupAndroid.bat b/Engine/Extras/Android/SetupAndroid.bat
--- a/Engine/Extras/Android/SetupAndroid.bat
+++ b/Engine/Extras/Android/SetupAndroid.bat
@@ -59,6 +59,10 @@ if DEFINED ANDROID_HOME (set a=1) ELSE (
 	set ANDROID_HOME=%STUDIO_SDK_PATH%
 	setx ANDROID_HOME "%STUDIO_SDK_PATH%"
 )
+if DEFINED ANDROID_SDK_HOME (set a=1) ELSE (
+	set ANDROID_SDK_HOME=%STUDIO_SDK_PATH%
+	setx ANDROID_SDK_HOME "%STUDIO_SDK_PATH%"
+)
 if DEFINED JAVA_HOME (set a=1) ELSE (
 	set JAVA_HOME=%STUDIO_PATH%\jre
 	setx JAVA_HOME "%STUDIO_PATH%\jre"
@@ -79,11 +83,11 @@ IF /I "%ERRORLEVEL%" NEQ "0" (
 	echo Added %PLATFORMTOOLS% to path
 )
 
-set SDKMANAGER=%STUDIO_SDK_PATH%\tools\bin\sdkmanager.bat
+set SDKMANAGER=%STUDIO_SDK_PATH%\cmdline-tools\latest\bin\sdkmanager.bat
 IF EXIST "%SDKMANAGER%" (
 	echo Using sdkmanager: %SDKMANAGER%
 ) ELSE (
-	set SDKMANAGER=%STUDIO_SDK_PATH%\cmdline-tools\latest\bin\sdkmanager.bat
+	set SDKMANAGER=%STUDIO_SDK_PATH%\tools\bin\sdkmanager.bat
 	IF EXIST "!SDKMANAGER!" (
 		echo Using sdkmanager: !SDKMANAGER!
 	) ELSE (
diff --git a/Engine/Extras/Android/SetupAndroid.command b/Engine/Extras/Android/SetupAndroid.command
--- a/Engine/Extras/Android/SetupAndroid.command
+++ b/Engine/Extras/Android/SetupAndroid.command
@@ -20,6 +20,12 @@ updateProfileFileFunc()
 		echo "export ANDROID_HOME=\"$STUDIO_SDK_PATH\"" >> "$1"
 	fi
  
+	if ! grep -q "export ANDROID_SDK_HOME=\"$STUDIO_SDK_PATH\"" "$1" 
+	then
+		echo >> "$1"
+		echo "export ANDROID_SDK_HOME=\"$STUDIO_SDK_PATH\"" >> "$1"
+	fi
+ 
 	if ! grep -q "export JAVA_HOME=\"$JAVA_HOME\"" "$1"
 	then
 		echo >> $1
@@ -94,9 +100,9 @@ export JAVA_HOME="$STUDIO_PATH/Contents/jre/jdk/Contents/Home"
 NDKINSTALLPATH="$STUDIO_SDK_PATH/ndk/${NDK_VERSION}"
 PLATFORMTOOLS="$STUDIO_SDK_PATH/platform-tools:$STUDIO_SDK_PATH/build-tools/${BUILDTOOLS_VERSION}:$STUDIO_SDK_PATH/tools/bin"
 adbPath=$(type -P "adb")
-SDKMANAGERPATH="$STUDIO_SDK_PATH/tools/bin"
+SDKMANAGERPATH="$STUDIO_SDK_PATH/cmdline-tools/latest/bin"
 if [ ! -d "$SDKMANAGERPATH" ]; then
-	SDKMANAGERPATH="$STUDIO_SDK_PATH/cmdline-tools/latest/bin"
+	SDKMANAGERPATH="$STUDIO_SDK_PATH/tools/bin"
 	if [ ! -d "$SDKMANAGERPATH" ]; then
 		echo Unable to locate sdkmanager. Did you run Android Studio and install cmdline-tools after installing?
 		${PAUSE}
diff --git a/Engine/Extras/Android/SetupAndroid.sh b/Engine/Extras/Android/SetupAndroid.sh
--- a/Engine/Extras/Android/SetupAndroid.sh
+++ b/Engine/Extras/Android/SetupAndroid.sh
@@ -67,6 +67,12 @@ then
 	echo "export ANDROID_HOME=\"$STUDIO_SDK_PATH\"" >>$HOME/.bashrc
 fi
 
+if ! grep -q "export ANDROID_SDK_HOME=\"$STUDIO_SDK_PATH\"" $HOME/.bashrc
+then
+	echo >>$HOME/.bashrc
+	echo "export ANDROID_SDK_HOME=\"$STUDIO_SDK_PATH\"" >>$HOME/.bashrc
+fi
+
 export JAVA_HOME="$STUDIO_PATH/jre"
 if ! grep -q "export JAVA_HOME=\"$JAVA_HOME\"" $HOME/.bashrc
 then
@@ -83,9 +89,9 @@ if [ "$retVal" == "" ]; then
 	echo Added $PLATFORMTOOLS to path
 fi
 
-SDKMANAGERPATH="$STUDIO_SDK_PATH/tools/bin"
+SDKMANAGERPATH="$STUDIO_SDK_PATH/cmdline-tools/latest/bin"
 if [ ! -d "$SDKMANAGERPATH" ]; then
-	SDKMANAGERPATH="$STUDIO_SDK_PATH/cmdline-tools/latest/bin"
+	SDKMANAGERPATH="$STUDIO_SDK_PATH/tools/bin"
 	if [ ! -d "$SDKMANAGERPATH" ]; then
 		echo Unable to locate sdkmanager. Did you run Android Studio and install cmdline-tools after installing?
 		${PAUSE}
