Use commandline tools 8.0 instead of latest for Android setup scripts

#android

diff --git a/Engine/Extras/Android/SetupAndroid.bat b/Engine/Extras/Android/SetupAndroid.bat
--- a/Engine/Extras/Android/SetupAndroid.bat
+++ b/Engine/Extras/Android/SetupAndroid.bat
@@ -83,7 +83,7 @@ IF /I "%ERRORLEVEL%" NEQ "0" (
 	echo Added %PLATFORMTOOLS% to path
 )
 
-set SDKMANAGER=%STUDIO_SDK_PATH%\cmdline-tools\latest\bin\sdkmanager.bat
+set SDKMANAGER=%STUDIO_SDK_PATH%\cmdline-tools\8.0\bin\sdkmanager.bat
 IF EXIST "%SDKMANAGER%" (
 	echo Using sdkmanager: %SDKMANAGER%
 ) ELSE (
diff --git a/Engine/Extras/Android/SetupAndroid.command b/Engine/Extras/Android/SetupAndroid.command
--- a/Engine/Extras/Android/SetupAndroid.command
+++ b/Engine/Extras/Android/SetupAndroid.command
@@ -97,7 +97,7 @@ export JAVA_HOME="$STUDIO_PATH/Contents/jre/jdk/Contents/Home"
 NDKINSTALLPATH="$STUDIO_SDK_PATH/ndk/${NDK_VERSION}"
 PLATFORMTOOLS="$STUDIO_SDK_PATH/platform-tools:$STUDIO_SDK_PATH/build-tools/${BUILDTOOLS_VERSION}:$STUDIO_SDK_PATH/tools/bin"
 adbPath=$(type -P "adb")
-SDKMANAGERPATH="$STUDIO_SDK_PATH/cmdline-tools/latest/bin"
+SDKMANAGERPATH="$STUDIO_SDK_PATH/cmdline-tools/8.0/bin"
 if [ ! -d "$SDKMANAGERPATH" ]; then
 	SDKMANAGERPATH="$STUDIO_SDK_PATH/tools/bin"
 	if [ ! -d "$SDKMANAGERPATH" ]; then
diff --git a/Engine/Extras/Android/SetupAndroid.sh b/Engine/Extras/Android/SetupAndroid.sh
--- a/Engine/Extras/Android/SetupAndroid.sh
+++ b/Engine/Extras/Android/SetupAndroid.sh
@@ -89,7 +89,7 @@ if [ "$retVal" == "" ]; then
 	echo Added $PLATFORMTOOLS to path
 fi
 
-SDKMANAGERPATH="$STUDIO_SDK_PATH/cmdline-tools/latest/bin"
+SDKMANAGERPATH="$STUDIO_SDK_PATH/cmdline-tools/8.0/bin"
 if [ ! -d "$SDKMANAGERPATH" ]; then
 	SDKMANAGERPATH="$STUDIO_SDK_PATH/tools/bin"
 	if [ ! -d "$SDKMANAGERPATH" ]; then
