Fix clearing of ANDROID_SDK_HOME env var in SetupAndroid.bat

#android

diff --git a/Engine/Extras/Android/SetupAndroid.bat b/Engine/Extras/Android/SetupAndroid.bat
--- a/Engine/Extras/Android/SetupAndroid.bat
+++ b/Engine/Extras/Android/SetupAndroid.bat
@@ -59,7 +59,7 @@ if DEFINED ANDROID_HOME (set a=1) ELSE (
 	set ANDROID_HOME=%STUDIO_SDK_PATH%
 	setx ANDROID_HOME "%STUDIO_SDK_PATH%"
 )
-if DEFINED ANDROID_SDK_HOME (set a=1) ELSE (
+if DEFINED ANDROID_SDK_HOME (
 	set ANDROID_SDK_HOME=
 	setx ANDROID_SDK_HOME ""
 )
