Fix location of JAVA_HOME for Mac OS

#android

diff --git a/Engine/Extras/Android/SetupAndroid.command b/Engine/Extras/Android/SetupAndroid.command
--- a/Engine/Extras/Android/SetupAndroid.command
+++ b/Engine/Extras/Android/SetupAndroid.command
@@ -95,7 +95,7 @@ if [ ! -d "$STUDIO_SDK_PATH" ]; then
 fi
 echo Android Studio SDK Path: $STUDIO_SDK_PATH
 
-export JAVA_HOME="$STUDIO_PATH/Contents/jbr/jdk/Contents/Home"
+export JAVA_HOME="$STUDIO_PATH/Contents/jbr/Contents/Home"
 NDKINSTALLPATH="$STUDIO_SDK_PATH/ndk/${NDK_VERSION}"
 PLATFORMTOOLS="$STUDIO_SDK_PATH/platform-tools:$STUDIO_SDK_PATH/build-tools/${BUILDTOOLS_VERSION}:$STUDIO_SDK_PATH/tools/bin"
 adbPath=$(type -P "adb")
