PR #11394: Fix Android Studio path detection for installs from JetBrains Toolbox

#android

diff --git a/Engine/Extras/Android/SetupAndroid.bat b/Engine/Extras/Android/SetupAndroid.bat
--- a/Engine/Extras/Android/SetupAndroid.bat
+++ b/Engine/Extras/Android/SetupAndroid.bat
@@ -22,9 +22,14 @@ if "%BUILDTOOLS_VERSION%" == "" SET BUILDTOOLS_VERSION=33.0.1
 if "%CMAKE_VERSION%" == "" SET CMAKE_VERSION=3.22.1
 if "%NDK_VERSION%" == "" SET NDK_VERSION=25.1.8937393
 
-
+rem first try looking in the Windows registry for Android Studio location
 FOR /F "tokens=2*" %%A IN ('REG.exe query "%KEY_NAME%" /v "%VALUE_NAME%"') DO (set STUDIO_PATH=%%B)
 
+rem Some installs, like JetBrains Toolbox, may not place an entry in the registry so try an alternate location
+if "%STUDIO_PATH%" == "" (
+	set STUDIO_PATH=%LOCALAPPDATA%\Programs\Android Studio
+)
+
 IF EXIST "%STUDIO_PATH%" (
 	echo.
 	) ELSE (
