Fixed Installing Android SDK from editor Platform menu does not take until after PC restart

#android

diff --git a/Engine/Extras/Android/SetupAndroid.bat b/Engine/Extras/Android/SetupAndroid.bat
--- a/Engine/Extras/Android/SetupAndroid.bat
+++ b/Engine/Extras/Android/SetupAndroid.bat
@@ -60,17 +60,17 @@ if "%STUDIO_SDK_PATH%" == "" (
 )
 echo Android Studio SDK Path: %STUDIO_SDK_PATH%
 
-if DEFINED ANDROID_HOME (set a=1) ELSE (
+if not defined ANDROID_HOME (
 	set ANDROID_HOME=%STUDIO_SDK_PATH%
-	setx ANDROID_HOME "%STUDIO_SDK_PATH%"
+	powershell -command "[Environment]::SetEnvironmentVariable('ANDROID_HOME', '%STUDIO_SDK_PATH%', 'User')"
 )
-if DEFINED ANDROID_SDK_HOME (
+if defined ANDROID_SDK_HOME (
 	set ANDROID_SDK_HOME=
-	setx ANDROID_SDK_HOME ""
+	powershell -command "[Environment]::SetEnvironmentVariable('ANDROID_SDK_HOME', [NullString]::Value, 'User')"
 )
-if DEFINED JAVA_HOME (set a=1) ELSE (
+if not defined JAVA_HOME (
 	set JAVA_HOME=%STUDIO_PATH%\jbr
-	setx JAVA_HOME "%STUDIO_PATH%\jbr"
+	powershell -command "[Environment]::SetEnvironmentVariable('JAVA_HOME', '%STUDIO_PATH%\jbr', 'User')"
 )
 set NDKINSTALLPATH=%STUDIO_SDK_PATH%\ndk\%NDK_VERSION%
 set PLATFORMTOOLS=%STUDIO_SDK_PATH%\platform-tools;%STUDIO_SDK_PATH%\tools
@@ -84,7 +84,7 @@ FOR /F "tokens=2*" %%A IN ('REG.exe query "%KEY_NAME%" /v "%VALUE_NAME%"') DO (s
 where.exe /Q adb.exe
 IF /I "%ERRORLEVEL%" NEQ "0" (
 	echo Current user path: "%USERPATH%"
-	setx PATH "%USERPATH%;%PLATFORMTOOLS%"
+	powershell -command "[Environment]::SetEnvironmentVariable('PATH', '%USERPATH%;%PLATFORMTOOLS%', 'User')"
 	echo Added %PLATFORMTOOLS% to path
 )
 
@@ -112,8 +112,8 @@ IF /I "%ERRORLEVEL%" NEQ "0" (
 
 if EXIST "%NDKINSTALLPATH%" (
 	echo Success!
-	setx NDKROOT "%NDKINSTALLPATH%"
-	setx NDK_ROOT "%NDKINSTALLPATH%"
+	powershell -command "[Environment]::SetEnvironmentVariable('NDKROOT', '%NDKINSTALLPATH%', 'User')"
+	powershell -command "[Environment]::SetEnvironmentVariable('NDK_ROOT', '%NDKINSTALLPATH%', 'User')"
 ) ELSE (
 	echo Update failed. Did you accept the license agreement?
 	%PAUSE%
