Fixed Force Reinstall SDK for Android failing on a Mac device on Apple Silicon

#android

#tests tested uname returns aarch64 and arm64

diff --git a/Engine/Extras/Android/SetupAndroidEmulator.command b/Engine/Extras/Android/SetupAndroidEmulator.command
--- a/Engine/Extras/Android/SetupAndroidEmulator.command
+++ b/Engine/Extras/Android/SetupAndroidEmulator.command
@@ -52,6 +52,9 @@ if [ ! -d "$CMDLINETOOLS_PATH" ]; then
 fi
 
 ABI=$(uname -m)
+if [ "${ABI}" == "aarch64" ] || [ "${ABI}" == "arm64" ]; then
+	ABI="arm64-v8a"
+fi
 if [ "${ABI}" != "x86_64" ] && [ "${ABI}" != "arm64-v8a" ]; then
 	echo "Unsupported architecture ${ABI}."
 	${PAUSE}
diff --git a/Engine/Extras/Android/SetupAndroidEmulator.sh b/Engine/Extras/Android/SetupAndroidEmulator.sh
--- a/Engine/Extras/Android/SetupAndroidEmulator.sh
+++ b/Engine/Extras/Android/SetupAndroidEmulator.sh
@@ -58,6 +58,9 @@ if [ ! -d "${CMDLINETOOLS_PATH}" ]; then
 fi
 
 ABI=$(uname -m)
+if [ "${ABI}" == "aarch64" ] || [ "${ABI}" == "arm64" ]; then
+	ABI="arm64-v8a"
+fi
 if [ "${ABI}" != "x86_64" ] && [ "${ABI}" != "arm64-v8a" ]; then
 	echo "Unsupported architecture ${ABI}."
 	${PAUSE}
