Update SetupAndroid to NDK to r27c

diff --git a/Engine/Extras/Android/SetupAndroid.bat b/Engine/Extras/Android/SetupAndroid.bat
--- a/Engine/Extras/Android/SetupAndroid.bat
+++ b/Engine/Extras/Android/SetupAndroid.bat
@@ -20,7 +20,7 @@ rem hardcoded versions for compatibility with non-Turnkey manual running
 if "%PLATFORMS_VERSION%" == "" SET PLATFORMS_VERSION=android-34
 if "%BUILDTOOLS_VERSION%" == "" SET BUILDTOOLS_VERSION=34.0.0
 if "%CMAKE_VERSION%" == "" SET CMAKE_VERSION=3.22.1
-if "%NDK_VERSION%" == "" SET NDK_VERSION=25.1.8937393
+if "%NDK_VERSION%" == "" SET NDK_VERSION=27.2.12479018
 
 rem first try looking in the Windows registry for Android Studio location
 FOR /F "tokens=2*" %%A IN ('REG.exe query "%KEY_NAME%" /v "%VALUE_NAME%"') DO (set STUDIO_PATH=%%B)
diff --git a/Engine/Extras/Android/SetupAndroid.command b/Engine/Extras/Android/SetupAndroid.command
--- a/Engine/Extras/Android/SetupAndroid.command
+++ b/Engine/Extras/Android/SetupAndroid.command
@@ -68,7 +68,7 @@ if [[ -z "${CMAKE_VERSION}" ]]; then
     CMAKE_VERSION="3.22.1"
 fi
 if [[ -z "${NDK_VERSION}" ]]; then
-    NDK_VERSION="25.1.8937393"
+    NDK_VERSION="27.2.12479018"
 fi
 
 
diff --git a/Engine/Extras/Android/SetupAndroid.sh b/Engine/Extras/Android/SetupAndroid.sh
--- a/Engine/Extras/Android/SetupAndroid.sh
+++ b/Engine/Extras/Android/SetupAndroid.sh
@@ -36,7 +36,7 @@ if [[ -z "${CMAKE_VERSION}" ]]; then
     CMAKE_VERSION="3.22.1"
 fi
 if [[ -z "${NDK_VERSION}" ]]; then
-    NDK_VERSION="25.1.8937393"
+    NDK_VERSION="27.2.12479018"
 fi
 
 
