Made undoing UE changes to .bash_profile on Mac handle more cases in a more compatible way

#android

diff --git a/Engine/Extras/Android/SetupAndroid.command b/Engine/Extras/Android/SetupAndroid.command
--- a/Engine/Extras/Android/SetupAndroid.command
+++ b/Engine/Extras/Android/SetupAndroid.command
@@ -55,22 +55,41 @@ updateProfileFile()
 	fi
 }
 
-# $1 parameter should be the profile file
+# $1 parameter should be the profile contents to remove
+# $2 parameter should be the profile file
 undoProfileFile()
 {
-	if [ ! -f "${1}" ]; then
-		return
+	if [ ! -f "${2}" ]; then
+		return 0
 	fi
+		
+	local code=\
+'{
+	data = data $0 ORS
+}
+END {
+	i = index(data, text)
+	if (i != 0)
+	{
+		data = substr(data, 1, i - 1) substr(data, i + length(text))
+		printf "%s", data > out
+		exit close(out)
+	}
+		
+	exit 1
+}'
 	
-	echo "Undoing UE Android dev env variables from ${1}"
-	
-	if sed "s|\nexport ANDROID_HOME=\"$STUDIO_SDK_PATH\"\n\nexport JAVA_HOME=\"$JAVA_HOME\"\n\nexport PATH=\"\$PATH:$PLATFORMTOOLS\"\n\nexport NDKROOT=\"$NDKINSTALLPATH\"\nexport NDK_ROOT=\"$NDKINSTALLPATH\"||g" "$1"; then
-		sed -i "s|\nexport ANDROID_HOME=\"$STUDIO_SDK_PATH\"\n\nexport JAVA_HOME=\"$JAVA_HOME\"\n\nexport PATH=\"\$PATH:$PLATFORMTOOLS\"\n\nexport NDKROOT=\"$NDKINSTALLPATH\"\nexport NDK_ROOT=\"$NDKINSTALLPATH\"||g" "$1"
+	if awk -v text="${1}" -v out="${2}" "${code}" "${2}"; then
+		echo "Removed UE Android developer environment variables from ${2}"
+		
+		if [ ! -s "${2}" ]; then
+			rm "${2}"
+		fi
+		
+		return 0
 	fi
 	
-	if [ ! -s "${1}" ] ; then
-		rm "${1}"
-    fi
+	return 1
 }
 
 if [[ ${ARG5} == "-noninteractive" ]]; then
@@ -162,7 +181,10 @@ fi
 updateProfileFile ~/.zshrc
 
 if [ "${SHELL}" == "/bin/zsh" ]; then
-	undoProfileFile ~/.bash_profile
+	undoProfileFile "\nexport ANDROID_HOME=\"$STUDIO_SDK_PATH\"\n\nexport JAVA_HOME=\"$JAVA_HOME\"\n\nexport PATH=\"\$PATH:$PLATFORMTOOLS\"\n\nexport NDKROOT=\"$STUDIO_SDK_PATH/ndk/25.1.8937393\"\nexport NDK_ROOT=\"$STUDIO_SDK_PATH/ndk/25.1.8937393\"\n" ~/.bash_profile ||\
+	undoProfileFile "\nexport ANDROID_HOME=\"$STUDIO_SDK_PATH\"\n\nexport JAVA_HOME=\"$JAVA_HOME\"\n\nexport NDKROOT=\"$STUDIO_SDK_PATH/ndk/25.1.8937393\"\nexport NDK_ROOT=\"$STUDIO_SDK_PATH/ndk/25.1.8937393\"\n" ~/.bash_profile ||\
+	undoProfileFile "\nexport ANDROID_HOME=\"$STUDIO_SDK_PATH\"\n\nexport JAVA_HOME=\"$JAVA_HOME\"\n\nexport PATH=\"\$PATH:$PLATFORMTOOLS\"\n\nexport NDKROOT=\"$NDKINSTALLPATH\"\nexport NDK_ROOT=\"$NDKINSTALLPATH\"\n" ~/.bash_profile ||\
+	undoProfileFile "\nexport ANDROID_HOME=\"$STUDIO_SDK_PATH\"\n\nexport JAVA_HOME=\"$JAVA_HOME\"\n\nexport NDKROOT=\"$NDKINSTALLPATH\"\nexport NDK_ROOT=\"$NDKINSTALLPATH\"\n" ~/.bash_profile
 fi
 
 echo Success!
